[RaceDetector]
- Avoid instrumenting memory access to stack allocated values that are never used outside the function
- This has a drastic reduction on the total number of calls to the instrumentation framework and improves performance considerably
- Fix a bug in the Microsoft ABI in recently introduced barrier that caused an unaligned CAS 64-bit on a 32-bit pointer

diff --git a/Engine/Programs/UnrealInstrumentation/Compiler/CustomMemoryInstrumentation.cpp b/Engine/Programs/UnrealInstrumentation/Compiler/CustomMemoryInstrumentation.cpp
--- a/Engine/Programs/UnrealInstrumentation/Compiler/CustomMemoryInstrumentation.cpp
+++ b/Engine/Programs/UnrealInstrumentation/Compiler/CustomMemoryInstrumentation.cpp
@@ -779,12 +779,29 @@ bool CustomMemoryInstrumentationPass::instrumentMSVCAtomicCallSite(
                                     InstrumentFunction, Sret);
 }
 
+bool CustomMemoryInstrumentationPass::shouldInstrumentAddr(Value *Addr)
+{
+  // if the variable is on stack and is never captured, we don't need to instrument it.
+  if (isa<AllocaInst>(getUnderlyingObject(Addr)) &&
+	  !PointerMayBeCaptured(Addr, true, true)) 
+  {
+	  return false;
+  }
+
+  return true;
+}
+
 bool CustomMemoryInstrumentationPass::instrumentStore(StoreInst *Inst,
                                                       bool SkipNonAtomics) {
   InstrumentationIRBuilder Builder(Inst);
 
   Value *Addr = Inst->getPointerOperand();
 
+  if (!shouldInstrumentAddr(Addr))
+  {
+    return false;
+  }
+
   // Special case for virtual table pointer updates.
   if (MDNode *Metadata = Inst->getMetadata(LLVMContext::MD_tbaa)) {
     if (Metadata->isTBAAVtableAccess()) {
@@ -829,17 +846,24 @@ bool CustomMemoryInstrumentationPass::instrumentLoad(LoadInst *Inst,
                                                      bool SkipNonAtomics) {
   InstrumentationIRBuilder Builder(Inst);
   
+  Value* Addr = Inst->getPointerOperand();
+
+  if (!shouldInstrumentAddr(Addr))
+  {
+     return false;
+  }
+
   // Special case for virtual table pointer reads.
   if (MDNode *Metadata = Inst->getMetadata(LLVMContext::MD_tbaa)) {
     if (Metadata->isTBAAVtableAccess()) {
-      Builder.CreateCall(LoadVPtrInstrumentFunction, Inst->getPointerOperand());
+      Builder.CreateCall(LoadVPtrInstrumentFunction, Addr);
       return true;
     }
   }
 
   Value *Ptr =
       Builder.CreateCast(Instruction::CastOps::PtrToInt,
-                         Inst->getPointerOperand(), Builder.getInt64Ty());
+						 Addr, Builder.getInt64Ty());
   uint32_t Size =
       CurrentModule->getDataLayout().getTypeStoreSize(Inst->getType());
 
@@ -850,7 +874,7 @@ bool CustomMemoryInstrumentationPass::instrumentLoad(LoadInst *Inst,
         ConstantInt::get(Builder.getInt8Ty(), MemoryOrderFromInst(Inst));
 
     return instrumentAtomicMemoryInst(
-        Builder, Inst, Inst->getPointerOperand(), nullptr /* value */,
+        Builder, Inst, Addr, nullptr /* value */,
         MemoryOrder, AtomicLoadInstrumentFunctions[FunctionIndexFromSize(Size)],
         nullptr);
   } else if (SkipNonAtomics) {
diff --git a/Engine/Programs/UnrealInstrumentation/Compiler/CustomMemoryInstrumentation.h b/Engine/Programs/UnrealInstrumentation/Compiler/CustomMemoryInstrumentation.h
--- a/Engine/Programs/UnrealInstrumentation/Compiler/CustomMemoryInstrumentation.h
+++ b/Engine/Programs/UnrealInstrumentation/Compiler/CustomMemoryInstrumentation.h
@@ -153,6 +153,7 @@ public:
 
 private:
   bool instrumentModule(Module &M);
+  bool shouldInstrumentAddr(Value* Addr);
   bool shouldInstrumentModule(Module &M);
   bool shouldInstrumentFunction(Function &F);
 
