[verse-clang] Pull in the UnrealInstrumentation/Compiler (UE's TSAN) into verse-clang.

Now that we have multi-patch support, this is easy to do.

#Synced-CL 52168512

diff --git a/Engine/Programs/UnrealInstrumentation/Compiler/CustomMemoryInstrumentation.cpp b/Engine/Programs/UnrealInstrumentation/Compiler/CustomMemoryInstrumentation.cpp
--- a/Engine/Programs/UnrealInstrumentation/Compiler/CustomMemoryInstrumentation.cpp
+++ b/Engine/Programs/UnrealInstrumentation/Compiler/CustomMemoryInstrumentation.cpp
@@ -33,7 +33,11 @@
 #include "llvm/Support/ModRef.h"
 #include "llvm/Support/Regex.h"
 #include "llvm/Support/SourceMgr.h"
+#if LLVM_VERSION_MAJOR >= 20
+#include "llvm/Transforms/Utils/Instrumentation.h"
+#else
 #include "llvm/Transforms/Instrumentation.h"
+#endif
 #include "llvm/Transforms/Utils/BasicBlockUtils.h"
 #include "llvm/Transforms/Utils/EscapeEnumerator.h"
 #include "llvm/Transforms/Utils/Local.h"
@@ -134,8 +138,8 @@ FAtomicMemoryOrder MemoryOrderFromLLVMOrdering(const AtomicOrdering &Ordering) {
   case AtomicOrdering::SequentiallyConsistent:
     return MEMORY_ORDER_SEQ_CST;
   case AtomicOrdering::NotAtomic:
-  default:
     assert(false);
+    break;
   }
   llvm_unreachable("Should have a memory order.");
 }
@@ -434,7 +438,11 @@ void CustomMemoryInstrumentationPass::cacheInstrumentationFunctions(Module &M) {
     AtomicStoreInstrumentFunctions[FunctionIndexFromSize(i)] =
         M.getOrInsertFunction(
             FuncName, Attrs, Builder.getVoidTy(),     // Return void.
-            Builder.getIntNTy(i * 8)->getPointerTo(), // Atomic pointer.
+#if LLVM_VERSION_MAJOR >= 20
+            Builder.getPtrTy(),
+#else
+            Builder.getIntNTy(i * 8)->getPointerTo(),
+#endif // Atomic pointer.
             Builder.getIntNTy(i * 8),                 // Value to store.
             Builder.getInt8Ty()                       // Memory order.
         );
@@ -445,7 +453,11 @@ void CustomMemoryInstrumentationPass::cacheInstrumentationFunctions(Module &M) {
     AtomicLoadInstrumentFunctions[FunctionIndexFromSize(i)] =
         M.getOrInsertFunction(
             FuncName, Attrs, Builder.getIntNTy(i * 8), // Return loaded value.
-            Builder.getIntNTy(i * 8)->getPointerTo(),  // Atomic pointer.
+#if LLVM_VERSION_MAJOR >= 20
+            Builder.getPtrTy(),
+#else
+            Builder.getIntNTy(i * 8)->getPointerTo(),
+#endif  // Atomic pointer.
             Builder.getInt8Ty()                        // Memory order.
         );
   }
@@ -454,7 +466,11 @@ void CustomMemoryInstrumentationPass::cacheInstrumentationFunctions(Module &M) {
     AtomicExchangeInstrumentFunctions[FunctionIndexFromSize(i)] =
         M.getOrInsertFunction(
             FuncName, Attrs, Builder.getIntNTy(i * 8), // Return previous value.
-            Builder.getIntNTy(i * 8)->getPointerTo(),  // Atomic pointer.
+#if LLVM_VERSION_MAJOR >= 20
+            Builder.getPtrTy(),
+#else
+            Builder.getIntNTy(i * 8)->getPointerTo(),
+#endif  // Atomic pointer.
             Builder.getIntNTy(i * 8),                  // Value to store.
             Builder.getInt8Ty()                        // Memory order.
         );
@@ -464,8 +480,16 @@ void CustomMemoryInstrumentationPass::cacheInstrumentationFunctions(Module &M) {
     AtomicCompareExchangeInstrumentFunctions[FunctionIndexFromSize(i)] =
         M.getOrInsertFunction(
             FuncName, Attrs, Builder.getIntNTy(i * 8), // Return previous value.
-            Builder.getIntNTy(i * 8)->getPointerTo(),  // Atomic pointer.
-            Builder.getIntNTy(i * 8)->getPointerTo(),  // Expected pointer.
+#if LLVM_VERSION_MAJOR >= 20
+            Builder.getPtrTy(),
+#else
+            Builder.getIntNTy(i * 8)->getPointerTo(),
+#endif  // Atomic pointer.
+#if LLVM_VERSION_MAJOR >= 20
+            Builder.getPtrTy(),
+#else
+            Builder.getIntNTy(i * 8)->getPointerTo(),
+#endif  // Expected pointer.
             Builder.getIntNTy(i * 8),                  // Value to store.
             Builder.getInt8Ty(),                       // Success memory order.
             Builder.getInt8Ty()                        // Failure memory order.
@@ -473,7 +497,7 @@ void CustomMemoryInstrumentationPass::cacheInstrumentationFunctions(Module &M) {
   }
 
   for (size_t i = 1; i <= MAX_ATOMIC_SIZE; i *= 2) {
-    for (int b = 0; b < RMWBinOp::LAST_BINOP; ++b) {
+    for (unsigned b = 0; b < RMWBinOp::LAST_BINOP; ++b) {
       if (!IsRMWOpHandled((RMWBinOp)b)) {
         continue;
       }
@@ -488,7 +512,11 @@ void CustomMemoryInstrumentationPass::cacheInstrumentationFunctions(Module &M) {
           M.getOrInsertFunction(
               FuncName, Attrs,
               Builder.getIntNTy(i * 8),                 // Return previous value.
-              Builder.getIntNTy(i * 8)->getPointerTo(), // Atomic pointer.
+  #if LLVM_VERSION_MAJOR >= 20
+            Builder.getPtrTy(),
+#else
+            Builder.getIntNTy(i * 8)->getPointerTo(),
+#endif // Atomic pointer.
               Builder.getIntNTy(i * 8),                 // Value to add.
               Builder.getInt8Ty()                       // Memory order.
           );
@@ -641,9 +669,15 @@ bool CustomMemoryInstrumentationPass::instrumentModule(Module &M) {
 }
 
 bool CustomMemoryInstrumentationPass::instrumentFunctionEntry(Function &F) {
+#if LLVM_VERSION_MAJOR >= 20
+  InstrumentationIRBuilder Builder(&*F.getEntryBlock().getFirstNonPHIIt());
+  Value *ReturnAddress = Builder.CreateCall(
+      Intrinsic::getOrInsertDeclaration(F.getParent(), Intrinsic::returnaddress),
+#else
   InstrumentationIRBuilder Builder(F.getEntryBlock().getFirstNonPHI());
   Value *ReturnAddress = Builder.CreateCall(
       Intrinsic::getDeclaration(F.getParent(), Intrinsic::returnaddress),
+#endif
       Builder.getInt32(0));
   Builder.CreateCall(FuncEntryInstrumentFunction, ReturnAddress);
   return true;
@@ -675,9 +709,8 @@ CustomMemoryInstrumentationPass::getInstrumentFunctionForMSVCAtomicCallSite(
   case CALL_SITE_RMW:
     return AtomicRMWInstrumentFunctions[CallSite.RMWOp]
                                        [FunctionIndexFromSize(Size)];
-  default:
-    llvm_unreachable("Should be handled");
   }
+  llvm_unreachable("Should be handled");
 }
 
 bool CustomMemoryInstrumentationPass::instrumentMSVCAtomicCallSite(
@@ -705,8 +738,13 @@ bool CustomMemoryInstrumentationPass::instrumentMSVCAtomicCallSite(
   }
 
   Value *Ptr =
+#if LLVM_VERSION_MAJOR >= 20
+      Builder.CreatePointerCast(Inst->getArgOperand(CallSite.PtrOperand),
+                                Builder.getPtrTy());
+#else
       Builder.CreatePointerCast(Inst->getArgOperand(CallSite.PtrOperand),
                                 Builder.getIntNTy(Size * 8)->getPointerTo());
+#endif
 
   Value *Val = nullptr;
   if (CallSite.StoreValueOperand) {
@@ -731,9 +769,15 @@ bool CustomMemoryInstrumentationPass::instrumentMSVCAtomicCallSite(
 
   Value *Expected = nullptr;
   if (CallSite.ExpectedOperand) {
+#if LLVM_VERSION_MAJOR >= 20
+    Expected = Builder.CreatePointerCast(
+        Inst->getArgOperand(*CallSite.ExpectedOperand),
+        Builder.getPtrTy());
+#else
     Expected = Builder.CreatePointerCast(
         Inst->getArgOperand(*CallSite.ExpectedOperand),
         Builder.getIntNTy(Size * 8)->getPointerTo());
+#endif
   }
 
   Value *MemoryOrder = nullptr;
diff --git a/Engine/Programs/UnrealInstrumentation/Compiler/CustomMemoryInstrumentation.h b/Engine/Programs/UnrealInstrumentation/Compiler/CustomMemoryInstrumentation.h
--- a/Engine/Programs/UnrealInstrumentation/Compiler/CustomMemoryInstrumentation.h
+++ b/Engine/Programs/UnrealInstrumentation/Compiler/CustomMemoryInstrumentation.h
@@ -6,12 +6,15 @@
 #define LLVM_TRANSFORMS_CUSTOMMEMORYINSTRUMENTATION_H
 
 #include "llvm/IR/PassManager.h"
+#if LLVM_VERSION_MAJOR >= 20
+#include "llvm/Transforms/Utils/Instrumentation.h"
+#else
 #include "llvm/Transforms/Instrumentation.h"
+#endif
 #include "llvm/ADT/StringSet.h"
 #include "llvm/Support/Regex.h"
 #include "llvm/ADT/SmallVector.h"
 #include "llvm/IR/IntrinsicInst.h"
-#include <intrin.h>
 
 namespace llvm {
 
