[verse-clang] Disable instrumentation for TSAN intrinsics.

Fixes a crash as atomics are used in the TSAN intrinsic functions, and the AutoRTFM passes were generating closed variants of these.

#Synced-CL 52345207

diff --git a/Engine/Programs/UnrealInstrumentation/Compiler/CustomMemoryInstrumentation.cpp b/Engine/Programs/UnrealInstrumentation/Compiler/CustomMemoryInstrumentation.cpp
--- a/Engine/Programs/UnrealInstrumentation/Compiler/CustomMemoryInstrumentation.cpp
+++ b/Engine/Programs/UnrealInstrumentation/Compiler/CustomMemoryInstrumentation.cpp
@@ -38,6 +38,9 @@
 #else
 #include "llvm/Transforms/Instrumentation.h"
 #endif
+#ifdef LLVM_HAS_AUTORTFM
+#include "AutoRTFM/Constants.h"
+#endif
 #include "llvm/Transforms/Utils/BasicBlockUtils.h"
 #include "llvm/Transforms/Utils/EscapeEnumerator.h"
 #include "llvm/Transforms/Utils/Local.h"
@@ -537,6 +540,19 @@ void CustomMemoryInstrumentationPass::cacheInstrumentationFunctions(Module &M) {
       }
     }
   }
+
+#ifdef LLVM_HAS_AUTORTFM
+  // Mark all __Instrument_* functions as autortfm_mode_internal so that
+  // AutoRTFM does not generate or execute closed variants of these
+  // instrumentation functions.
+  for (auto &F : M) {
+    if (F.getName().starts_with("__Instrument_")) {
+      F.addFnAttr(Attribute::get(
+          M.getContext(), Attribute::AutoRTFM,
+          static_cast<uint64_t>(autortfm_mode_internal)));
+    }
+  }
+#endif
 }
 
 bool CustomMemoryInstrumentationPass::instrumentMSVCStandardLib(Module &M) {
